Oodle JPEG format support. Smaller UASSETs and much more hardened security focus by using Oodle as a backend.

diff --git a/Engine/Config/BaseEditor.ini b/Engine/Config/BaseEditor.ini
--- a/Engine/Config/BaseEditor.ini
+++ b/Engine/Config/BaseEditor.ini
@@ -256,6 +256,7 @@ Class=UnrealEd.ActorFactory
 AllowNonPowerOfTwoTextures=True
 ; Imported jpegs will be stored in packages in jpeg format until edited, this will help keep texture asset sizes down
 RetainJpegFormat=True
+EnableOodleJPEG=False
 
 [/Script/Cascade.CascadeConfiguration]
 +ModuleMenu_ModuleRejections=ParticleModuleTypeDataBase
diff --git a/Engine/Plugins/Interchange/Runtime/Source/Import/Private/InterchangeImportModule.cpp b/Engine/Plugins/Interchange/Runtime/Source/Import/Private/InterchangeImportModule.cpp
--- a/Engine/Plugins/Interchange/Runtime/Source/Import/Private/InterchangeImportModule.cpp
+++ b/Engine/Plugins/Interchange/Runtime/Source/Import/Private/InterchangeImportModule.cpp
@@ -32,6 +32,7 @@
 #include "Texture/InterchangeDDSTranslator.h"
 #include "Texture/InterchangeIESTranslator.h"
 #include "Texture/InterchangeJPGTranslator.h"
+#include "Texture/InterchangeOOJPGTranslator.h"
 #include "Texture/InterchangePCXTranslator.h"
 #include "Texture/InterchangePSDTranslator.h"
 #include "Texture/InterchangeTextureFactory.h"
@@ -86,6 +87,7 @@ void FInterchangeImportModule::StartupModule()
 		//Textures
 		InterchangeManager.RegisterTranslator(UInterchangeImageWrapperTranslator::StaticClass());
 		InterchangeManager.RegisterTranslator(UInterchangeDDSTranslator::StaticClass());
+		InterchangeManager.RegisterTranslator(UInterchangeOOJPGTranslator::StaticClass());
 		InterchangeManager.RegisterTranslator(UInterchangeJPGTranslator::StaticClass());
 		InterchangeManager.RegisterTranslator(UInterchangePCXTranslator::StaticClass());
 		InterchangeManager.RegisterTranslator(UInterchangePSDTranslator::StaticClass());
diff --git a/Engine/Plugins/Interchange/Runtime/Source/Import/Private/Texture/InterchangeOOJPGTranslator.cpp b/Engine/Plugins/Interchange/Runtime/Source/Import/Private/Texture/InterchangeOOJPGTranslator.cpp
new file mode 100644
--- /dev/null
+++ b/Engine/Plugins/Interchange/Runtime/Source/Import/Private/Texture/InterchangeOOJPGTranslator.cpp
@@ -0,0 +1,147 @@
+// Copyright Epic Games, Inc. All Rights Reserved. 
+#include "Texture/InterchangeOOJPGTranslator.h"
+
+#include "Algo/Find.h"
+#include "Containers/StaticArray.h"
+#include "Engine/Texture.h"
+#include "Engine/Texture2D.h"
+#include "IImageWrapper.h"
+#include "IImageWrapperModule.h"
+#include "ImageCoreUtils.h"
+#include "InterchangeImportLog.h"
+#include "InterchangeTextureNode.h"
+#include "Memory/SharedBuffer.h"
+#include "Misc/ConfigCacheIni.h"
+#include "Misc/FileHelper.h"
+#include "Misc/Paths.h"
+#include "Misc/ScopedSlowTask.h"
+#include "Modules/ModuleManager.h"
+#include "Nodes/InterchangeBaseNodeContainer.h"
+#include "Texture/TextureTranslatorUtilities.h"
+
+#include UE_INLINE_GENERATED_CPP_BY_NAME(InterchangeOOJPGTranslator)
+
+static bool GInterchangeEnableOOJPGImport = true;
+static FAutoConsoleVariableRef CCvarInterchangeEnableOOJPGImport(
+	TEXT("Interchange.FeatureFlags.Import.OOJPG"),
+	GInterchangeEnableOOJPGImport,
+	TEXT("Whether OOJPG support is enabled."),
+	ECVF_Default);
+
+TArray<FString> UInterchangeOOJPGTranslator::GetSupportedFormats() const
+{
+	TArray<FString> Formats;
+
+	if (GInterchangeEnableOOJPGImport || GIsAutomationTesting)
+	{
+		Formats.Reserve(3);
+
+		Formats.Add(TEXT("ooj;OOJPEG image"));
+		bool bEnableOodleJPEG = false;
+		GConfig->GetBool(TEXT("TextureImporter"), TEXT("EnableOodleJPEG"), bEnableOodleJPEG, GEditorIni);
+		if (bEnableOodleJPEG)
+		{
+			Formats.Add(TEXT("jpg;JPEG image"));
+			Formats.Add(TEXT("jpeg;JPEG image"));
+		}
+	}
+
+	return Formats;
+}
+
+bool UInterchangeOOJPGTranslator::Translate(UInterchangeBaseNodeContainer& BaseNodeContainer) const
+{
+	return UE::Interchange::FTextureTranslatorUtilities::Generic2DTextureTranslate(GetSourceData(), BaseNodeContainer);
+}
+
+TOptional<UE::Interchange::FImportImage> UInterchangeOOJPGTranslator::GetTexturePayloadData(const FString& /*PayloadKey*/, TOptional<FString>& AlternateTexturePath) const
+{
+	using namespace UE::Interchange;
+
+	TArray64<uint8> SourceDataBuffer;
+	if (!FTextureTranslatorUtilities::LoadSourceBuffer(*this, TEXT("OOJPEG"), SourceDataBuffer))
+	{
+		return {};
+	}
+
+	const bool bImportRaw = false;
+	return GetTexturePayloadImplementation(MoveTemp(SourceDataBuffer), bImportRaw);
+
+}
+
+bool UInterchangeOOJPGTranslator::SupportCompressedTexturePayloadData() const
+{
+	return true;
+}
+
+TOptional<UE::Interchange::FImportImage> UInterchangeOOJPGTranslator::GetCompressedTexturePayloadData(const FString& /*PayloadKey*/, TOptional<FString>& /*AlternateTexturePath*/) const
+{
+	using namespace UE::Interchange;
+
+	TArray64<uint8> SourceDataBuffer;
+	if (!FTextureTranslatorUtilities::LoadSourceBuffer(*this, TEXT("OOJPEG"), SourceDataBuffer))
+	{
+		return {};
+	}
+
+	const bool bImportRaw = true;
+	return GetTexturePayloadImplementation(MoveTemp(SourceDataBuffer), bImportRaw);
+}
+
+TOptional<UE::Interchange::FImportImage> UInterchangeOOJPGTranslator::GetTexturePayloadImplementation(TArray64<uint8>&& SourceDataBuffer, bool bShouldImportRaw) const
+{
+	using namespace UE::Interchange;
+
+	IImageWrapperModule& ImageWrapperModule = FModuleManager::LoadModuleChecked<IImageWrapperModule>(FName("ImageWrapper"));
+
+	//
+	// OOJPG
+	//
+	ETextureSourceCompressionFormat TscfFormat = ETextureSourceCompressionFormat::TSCF_OOJPEG; 
+	TSharedPtr<IImageWrapper> ImageWrapper = ImageWrapperModule.CreateImageWrapper(EImageFormat::OOJPEG);
+	if (!ImageWrapper.IsValid() || !ImageWrapper->SetCompressed(SourceDataBuffer.GetData(), SourceDataBuffer.Num()))
+	{
+		TscfFormat = ETextureSourceCompressionFormat::TSCF_JPEG; 
+		ImageWrapper = ImageWrapperModule.CreateImageWrapper(EImageFormat::JPEG);
+		if (!ImageWrapper.IsValid() || !ImageWrapper->SetCompressed(SourceDataBuffer.GetData(), SourceDataBuffer.Num()))
+		{
+			FTextureTranslatorUtilities::LogError(*this, NSLOCTEXT("InterchangeOOJPGTranslator", "DecodingFailed", "Failed to decode OOJPEG."));
+			return TOptional<UE::Interchange::FImportImage>();
+		}
+	}
+
+	// Select the texture's source format
+	ERawImageFormat::Type RawFormat = ImageWrapper->GetClosestRawImageFormat();
+	check( RawFormat != ERawImageFormat::Invalid );
+	ETextureSourceFormat TextureFormat = FImageCoreUtils::ConvertToTextureSourceFormat(RawFormat);
+
+	UE::Interchange::FImportImage PayloadData;
+
+	const bool bShouldAllocateRawDataBuffer = false;
+
+	PayloadData.Init2DWithParams(
+		ImageWrapper->GetWidth(),
+		ImageWrapper->GetHeight(),
+		TextureFormat,
+		ImageWrapper->GetSRGB(),
+		bShouldAllocateRawDataBuffer
+	);
+
+	TArray64<uint8> RawData;
+	if (bShouldImportRaw)
+	{
+		PayloadData.RawData = MakeUniqueBufferFromArray(ImageWrapper->GetCompressed());
+		PayloadData.RawDataCompressionFormat = TscfFormat;
+	}
+	else if (ImageWrapper->GetRaw(RawData))
+	{
+		PayloadData.RawData = MakeUniqueBufferFromArray(MoveTemp(RawData));
+	}
+	else
+	{
+		FTextureTranslatorUtilities::LogError(*this, NSLOCTEXT("InterchangeOOJPGTranslator", "DecodingFailed", "Failed to decode OOJPEG."));
+		return TOptional<UE::Interchange::FImportImage>();
+	}
+
+	return PayloadData;
+}
diff --git a/Engine/Plugins/Interchange/Runtime/Source/Import/Private/Texture/InterchangeTextureFactory.cpp b/Engine/Plugins/Interchange/Runtime/Source/Import/Private/Texture/InterchangeTextureFactory.cpp
--- a/Engine/Plugins/Interchange/Runtime/Source/Import/Private/Texture/InterchangeTextureFactory.cpp
+++ b/Engine/Plugins/Interchange/Runtime/Source/Import/Private/Texture/InterchangeTextureFactory.cpp
@@ -41,6 +41,7 @@
 #include "Serialization/EditorBulkData.h"
 #include "Texture/InterchangeBlockedTexturePayloadInterface.h"
 #include "Texture/InterchangeJPGTranslator.h"
+#include "Texture/InterchangeOOJPGTranslator.h"
 #include "Texture/InterchangeSlicedTexturePayloadInterface.h"
 #include "Texture/InterchangeTextureLightProfilePayloadInterface.h"
 #include "Texture/InterchangeTexturePayloadInterface.h"
@@ -678,7 +679,8 @@ namespace UE::Interchange::Private::InterchangeTextureFactory
 			{
 				if (BlockAndSourceDataFiles.IsEmpty())
 				{
-					if (Translator->GetClass() == UInterchangeJPGTranslator::StaticClass())
+					if (Translator->GetClass() == UInterchangeJPGTranslator::StaticClass() 
+						|| Translator->GetClass() == UInterchangeOOJPGTranslator::StaticClass())
 					{
 						// Honor setting from TextureImporter.RetainJpegFormat in Editor.ini if it exists (ideally we should deprecate this as it is confusing and probably not thread safe)
 						const bool bShouldImportRawCache = bShoudImportCompressedImage;
diff --git a/Engine/Plugins/Interchange/Runtime/Source/Import/Public/Texture/InterchangeOOJPGTranslator.h b/Engine/Plugins/Interchange/Runtime/Source/Import/Public/Texture/InterchangeOOJPGTranslator.h
new file mode 100644
--- /dev/null
+++ b/Engine/Plugins/Interchange/Runtime/Source/Import/Public/Texture/InterchangeOOJPGTranslator.h
@@ -0,0 +1,63 @@
+// Copyright Epic Games, Inc. All Rights Reserved.
+
+#pragma once
+
+#include "CoreMinimal.h"
+#include "UObject/Object.h"
+#include "UObject/ObjectMacros.h"
+#include "InterchangeTranslatorBase.h"
+#include "Nodes/InterchangeBaseNodeContainer.h"
+#include "Texture/InterchangeTexturePayloadData.h"
+#include "Texture/InterchangeTexturePayloadInterface.h"
+
+#include "InterchangeOOJPGTranslator.generated.h"
+
+UCLASS(BlueprintType)
+class INTERCHANGEIMPORT_API UInterchangeOOJPGTranslator : public UInterchangeTranslatorBase, public IInterchangeTexturePayloadInterface
+{
+	GENERATED_BODY()
+public:
+
+	virtual TArray<FString> GetSupportedFormats() const override;
+	virtual EInterchangeTranslatorAssetType GetSupportedAssetTypes() const override { return EInterchangeTranslatorAssetType::Textures; }
+
+	/**
+	 * Translate the associated source data into a node hold by the specified nodes container.
+	 *
+	 * @param BaseNodeContainer - The unreal objects descriptions container where to put the translated source data.
+	 * @return true if the translator can translate the source data, false otherwise.
+	 */
+	virtual bool Translate(UInterchangeBaseNodeContainer& BaseNodeContainer) const override;
+
+	/* IInterchangeTexturePayloadInterface Begin */
+
+	/**
+	 * Once the translation is done, the import process need a way to retrieve payload data.
+	 * This payload will be use by the factories to create the asset.
+	 *
+	 * @param PayloadKey - Unused. The translator uses its SourceData property to extract the payload.
+	 * @param AlternateTexturePath - Unused. The translator uses its SourceData property to extract the payload.
+	 * @return a PayloadData containing the import image data. The TOptional will not be set if there is an error.
+	 */
+	virtual TOptional<UE::Interchange::FImportImage> GetTexturePayloadData(const FString& PayloadKey, TOptional<FString>& AlternateTexturePath) const override;
+
+	virtual bool SupportCompressedTexturePayloadData() const override;
+
+	/**
+	 * Once the translation is done, the import process need a way to retrieve payload data.
+	 * This payload will be use by the factories to create the asset.
+	 *
+	 * @param PayloadKey - Unused. The translator uses its SourceData property to extract the payload.
+	 * @param AlternateTexturePath - Unused. The translator uses its SourceData property to extract the payload.
+	 * @return a PayloadData containing the import image data. The TOptional will not be set if there is an error.
+	 */
+	virtual TOptional<UE::Interchange::FImportImage> GetCompressedTexturePayloadData(const FString& PayloadKey, TOptional<FString>& AlternateTexturePath) const override;
+	
+	/* IInterchangeTexturePayloadInterface End */
+
+private:
+
+	TOptional<UE::Interchange::FImportImage> GetTexturePayloadImplementation(TArray64<uint8>&& SourceDataBuffer, bool bShouldImportRaw) const;
+};
+
+
diff --git a/Engine/Source/Developer/TextureBuild/Private/TextureBuildFunction.cpp b/Engine/Source/Developer/TextureBuild/Private/TextureBuildFunction.cpp
--- a/Engine/Source/Developer/TextureBuild/Private/TextureBuildFunction.cpp
+++ b/Engine/Source/Developer/TextureBuild/Private/TextureBuildFunction.cpp
@@ -218,6 +218,13 @@ static bool TryReadTextureSourceFromCompactBinary(FCbFieldView Source, UE::Deriv
 			ImageWrapper->GetRaw(SourceFormat == TSF_G8 ? ERGBFormat::Gray : ERGBFormat::BGRA, 8, IntermediateDecompressedData);
 		}
 		break;
+		case TSCF_OOJPEG:
+		{
+			TSharedPtr<IImageWrapper> ImageWrapper = FModuleManager::GetModuleChecked<IImageWrapperModule>(FName("ImageWrapper")).CreateImageWrapper(EImageFormat::OOJPEG);
+			ImageWrapper->SetCompressed((const uint8*)InputBuffer.GetData(), InputBuffer.GetSize());
+			ImageWrapper->GetRaw(SourceFormat == TSF_G8 ? ERGBFormat::Gray : ERGBFormat::BGRA, 8, IntermediateDecompressedData);
+		}
+		break;
 		case TSCF_PNG:
 		{
 			TSharedPtr<IImageWrapper> ImageWrapper = FModuleManager::GetModuleChecked<IImageWrapperModule>(FName("ImageWrapper")).CreateImageWrapper(EImageFormat::PNG);
diff --git a/Engine/Source/Editor/UnrealEd/Private/Factories/EditorFactories.cpp b/Engine/Source/Editor/UnrealEd/Private/Factories/EditorFactories.cpp
--- a/Engine/Source/Editor/UnrealEd/Private/Factories/EditorFactories.cpp
+++ b/Engine/Source/Editor/UnrealEd/Private/Factories/EditorFactories.cpp
@@ -2930,7 +2930,7 @@ bool UTextureFactory::ImportImage(const uint8* Buffer, int64 Length, FFeedbackCo
 			// branch for JPEG, if retaining the jpeg compressed data
 			// this is inside the DecompressImage branch even though we don't use the LoadedImage at all
 			//	 just to ensure that the jpeg will decode successfully
-			if (ImageFormat == EImageFormat::JPEG)
+			if (ImageFormat == EImageFormat::JPEG || ImageFormat == EImageFormat::OOJPEG)
 			{
 				// unusual loader, retains jpeg
 				bool bRetainJpegFormat = false;
@@ -2942,7 +2942,7 @@ bool UTextureFactory::ImportImage(const uint8* Buffer, int64 Length, FFeedbackCo
 				if ( bRetainJpegFormat)
 				{
 					// does not decode jpeg, just to get width/height :
-					TSharedPtr<IImageWrapper> JpegImageWrapper = ImageWrapperModule.CreateImageWrapper(EImageFormat::JPEG);
+					TSharedPtr<IImageWrapper> JpegImageWrapper = ImageWrapperModule.CreateImageWrapper(ImageFormat);
 					if (JpegImageWrapper.IsValid() && JpegImageWrapper->SetCompressed(Buffer, Length))
 					{
 						check( JpegImageWrapper->GetWidth() == LoadedImage.SizeX );
@@ -2961,9 +2961,17 @@ bool UTextureFactory::ImportImage(const uint8* Buffer, int64 Length, FFeedbackCo
 						);
 
 						OutImage.RawData.Append(Buffer, Length);
-						OutImage.RawDataCompressionFormat = ETextureSourceCompressionFormat::TSCF_JPEG;
 
-						UE_LOG(LogEditorFactories,Display,TEXT("JPEG imported and retained as JPEG in uasset."));
+						if (ImageFormat == EImageFormat::JPEG) 
+						{
+							OutImage.RawDataCompressionFormat = ETextureSourceCompressionFormat::TSCF_JPEG;
+							UE_LOG(LogEditorFactories,Display,TEXT("JPEG imported and retained as JPEG in uasset."));
+						}
+						else
+						{
+							OutImage.RawDataCompressionFormat = ETextureSourceCompressionFormat::TSCF_OOJPEG;
+							UE_LOG(LogEditorFactories,Display,TEXT("OOJPEG imported and retained as OOJPEG in uasset."));
+						}
 
 						return true;
 					}
@@ -5192,7 +5200,7 @@ bool UTextureExporterJPEG::SupportsObject(UObject* Object) const
 			}
 
 			// Check it has JPEG BulkData :
-			if ( Texture->Source.GetSourceCompression() == TSCF_JPEG &&
+			if ( (Texture->Source.GetSourceCompression() == TSCF_JPEG || Texture->Source.GetSourceCompression() == TSCF_OOJPEG) &&
 				Texture->Source.GetSizeOnDisk() > 0 )
 			{
 				ETextureSourceFormat TSF = Texture->Source.GetFormat();
@@ -5206,12 +5214,12 @@ bool UTextureExporterJPEG::SupportsObject(UObject* Object) const
 	return false;
 }
 
-bool UTextureExporterJPEG::ExportBinary( UObject* Object, const TCHAR* Type, FArchive& Ar, FFeedbackContext* Warn, int32 FileIndex, uint32 PortFlags )
+bool UTextureExporterJPEG::ExportBinary(UObject* Object, const TCHAR* Type, FArchive& Ar, FFeedbackContext* Warn, int32 FileIndex, uint32 PortFlags)
 {
 	UTexture2D* Texture = Cast<UTexture2D>(Object);
-	check( Texture != nullptr );
+	check(Texture != nullptr);
 
-	check( Texture->Source.GetSourceCompression() == TSCF_JPEG &&
+	check((Texture->Source.GetSourceCompression() == TSCF_JPEG || Texture->Source.GetSourceCompression() == TSCF_OOJPEG) &&
 			Texture->Source.GetSizeOnDisk() > 0 );
 
 	// just write the JPEG data we already have :
@@ -5219,9 +5227,19 @@ bool UTextureExporterJPEG::ExportBinary( UObject* Object, const TCHAR* Type, FAr
 	UE_LOG(LogEditorFactories, Display, TEXT("Exporting Texture as JPEG stored bits (no lossy decompress or recompress)."));
 
 	Texture->Source.OperateOnLoadedBulkData( [&](const FSharedBuffer& BulkDataBuffer) {
-		Ar.Serialize( const_cast<void *>(BulkDataBuffer.GetData()), BulkDataBuffer.GetSize());
+		if (Texture->Source.GetSourceCompression() == TSCF_JPEG)
+		{
+			Ar.Serialize(const_cast<void*>(BulkDataBuffer.GetData()), BulkDataBuffer.GetSize());
+		}
+		else
+		{
+			IImageWrapperModule& ImageWrapperModule = FModuleManager::LoadModuleChecked<IImageWrapperModule>(FName("ImageWrapper"));
+			TSharedPtr<IImageWrapper> ImageWrapper = ImageWrapperModule.CreateImageWrapper(EImageFormat::OOJPEG);
+			ImageWrapper->SetCompressed(BulkDataBuffer.GetData(), BulkDataBuffer.GetSize());
+			TArray64<uint8> ExportData = ImageWrapper->GetExportData();
+			Ar.Serialize(ExportData.GetData(), ExportData.Num());
+		}
 	} );
-
 	return true;
 }
 
diff --git a/Engine/Source/Runtime/Engine/Classes/Engine/TextureDefines.h b/Engine/Source/Runtime/Engine/Classes/Engine/TextureDefines.h
--- a/Engine/Source/Runtime/Engine/Classes/Engine/TextureDefines.h
+++ b/Engine/Source/Runtime/Engine/Classes/Engine/TextureDefines.h
@@ -318,6 +318,7 @@ enum ETextureSourceCompressionFormat : int
 	TSCF_None	UMETA(DisplayName = "None"),
 	TSCF_PNG	UMETA(DisplayName = "PNG"),
 	TSCF_JPEG	UMETA(DisplayName = "JPEG"),
+	TSCF_OOJPEG	UMETA(DisplayName = "Oodle JPEG"),
 
 	TSCF_MAX
 };
diff --git a/Engine/Source/Runtime/ImageWrapper/ImageWrapper.Build.cs b/Engine/Source/Runtime/ImageWrapper/ImageWrapper.Build.cs
--- a/Engine/Source/Runtime/ImageWrapper/ImageWrapper.Build.cs
+++ b/Engine/Source/Runtime/ImageWrapper/ImageWrapper.Build.cs
@@ -36,7 +36,9 @@ public class ImageWrapper : ModuleRules
 		AddEngineThirdPartyPrivateStaticDependencies(Target,
 			"zlib",
 			"UElibPNG",
-			"LibTiff"
+			"LibTiff",
+			"OodleDataCompression",
+			"OodleJPEG"
 		);
 
 		// Jpeg Decoding
diff --git a/Engine/Source/Runtime/ImageWrapper/Private/Formats/OOJpegImageWrapper.cpp b/Engine/Source/Runtime/ImageWrapper/Private/Formats/OOJpegImageWrapper.cpp
new file mode 100644
--- /dev/null
+++ b/Engine/Source/Runtime/ImageWrapper/Private/Formats/OOJpegImageWrapper.cpp
@@ -0,0 +1,301 @@
+// Copyright Epic Games, Inc. All Rights Reserved.
+
+#include "Formats/OOJpegImageWrapper.h"
+
+#include "Async/ParallelFor.h"
+#include "Math/Color.h"
+#include "Misc/ScopeLock.h"
+#include "ImageWrapperPrivate.h"
+#include "ImageCoreUtils.h"
+
+#if WITH_OOJPEG
+
+#include "oojpeg.h"
+
+static void* ue_malloc(size_t size) 
+{ 
+	return FMemory::Malloc(size); 
+}
+
+static void ue_free(void* ptr) 
+{ 
+	FMemory::Free(ptr); 
+}
+
+static void* ue_realloc(void* ptr, size_t size) 
+{ 
+	return FMemory::Realloc(ptr, size); 
+}
+
+
+FOOJpegImageWrapper::FOOJpegImageWrapper(int32 InNumComponents /* = 4 */)
+	: FImageWrapperBase()
+	, NumComponents(InNumComponents)
+{ 
+	// NumComponents == 1 for GrayscaleJPEG
+	static oojpeg_alloc_t UeAlloc = { ue_malloc, ue_free, ue_realloc };
+	oojpeg_set_alloc(&UeAlloc);
+}
+
+
+FOOJpegImageWrapper::~FOOJpegImageWrapper()
+{
+}
+
+/* FImageWrapperBase interface
+ *****************************************************************************/
+
+bool FOOJpegImageWrapper::SetCompressed(const void* InCompressedData, int64 InCompressedSize)
+{
+	// jpgd doesn't support 64-bit sizes.
+	if (InCompressedSize < 0 || InCompressedSize > MAX_uint32)
+	{
+		return false;
+	}
+
+	int x, y, n;
+	unsigned char* d;
+
+	// Is it a JPEG?
+	if (!FMemory::Memcmp(InCompressedData, "\xFF\xD8\xFF", 3)) 
+	{
+		// First encode it to a OOJPEG format
+		unsigned char* OojData = 0;
+		int OojDataSize = 0;
+		{
+			oojpeg_encode_context_t ctx = oojpeg_encode_jpeg_mem_threaded_start((uint8*)InCompressedData, InCompressedSize, 0);	
+			if (ctx.error) 
+			{
+				return false;
+			}
+			ParallelFor(ctx.jobs_to_run, [&](int idx) {
+				oojpeg_encode_jpeg_thread_run(&ctx, idx);
+			});
+			int err = oojpeg_encode_jpeg_mem_threaded_finish(&ctx, &OojData, &OojDataSize);
+			if (err) 
+			{
+				return false;
+			}
+		}
+
+		// Then decode to a raw image
+		{
+			oojpeg_decode_context_t ctx = oojpeg_decode_mem_threaded_start(OojData, OojDataSize, 0);
+			if (ctx.error)
+			{
+				FMemory::Free(OojData);
+				return false;
+			}
+			ParallelFor(ctx.jobs_to_run, [&](int idx) {
+				oojpeg_decode_thread_run(&ctx, idx);
+			});
+			d = oojpeg_decode_threaded_finish(&ctx, &x, &y, &n);
+		}
+		if (!d)
+		{
+			FMemory::Free(OojData);
+			return false;
+		}
+		FMemory::Free(d);
+
+		if (!FImageWrapperBase::SetCompressed(OojData, OojDataSize))
+		{
+			FMemory::Free(OojData);
+			return false;
+		}
+
+		FMemory::Free(OojData);
+	}
+	else 
+	{
+		oojpeg_decode_context_t ctx = oojpeg_decode_mem_threaded_start((uint8*)InCompressedData, (uint32)InCompressedSize, 0);
+		if (ctx.error)
+		{
+			return false;
+		}
+		ParallelFor(ctx.jobs_to_run, [&](int idx) {
+			oojpeg_decode_thread_run(&ctx, idx);
+		});
+		d = oojpeg_decode_threaded_finish(&ctx, &x, &y, &n);
+		if (!d)
+		{
+			return false;
+		}
+		FMemory::Free(d);
+
+		if (!FImageWrapperBase::SetCompressed(InCompressedData, InCompressedSize))
+		{
+			return false;
+		}
+	}
+
+
+	// We don't support 16 bit jpegs
+	BitDepth = 8;
+
+	Width = x;
+	Height = y;
+
+	switch (n)
+	{
+	case 1:
+		Format = ERGBFormat::Gray;
+		break;
+	case 4:
+		Format = ERGBFormat::RGBA;
+		break;
+	default:
+		return false;
+	}
+	
+	if ( ! FImageCoreUtils::IsImageImportPossible(Width,Height) )
+	{
+		SetError(TEXT("Image dimensions are not possible to import"));
+		return false;
+	}
+
+	return true;
+}
+
+
+// CanSetRawFormat returns true if SetRaw will accept this format
+bool FOOJpegImageWrapper::CanSetRawFormat(const ERGBFormat InFormat, const int32 InBitDepth) const
+{
+	return ((InFormat == ERGBFormat::RGBA || InFormat == ERGBFormat::BGRA || InFormat == ERGBFormat::Gray) && InBitDepth == 8);
+}
+
+// returns InFormat if supported, else maps to something supported
+ERawImageFormat::Type FOOJpegImageWrapper::GetSupportedRawFormat(const ERawImageFormat::Type InFormat) const
+{
+	switch(InFormat)
+	{
+	case ERawImageFormat::G8:
+	case ERawImageFormat::BGRA8:
+		return InFormat; // directly supported
+	case ERawImageFormat::BGRE8:
+	case ERawImageFormat::RGBA16:
+	case ERawImageFormat::RGBA16F:
+	case ERawImageFormat::RGBA32F:
+	case ERawImageFormat::G16:
+	case ERawImageFormat::R16F:
+	case ERawImageFormat::R32F:
+		return ERawImageFormat::BGRA8; // needs conversion
+	default:
+		check(0);
+		return ERawImageFormat::BGRA8;
+	};
+}
+
+void FOOJpegImageWrapper::Compress(int32 Quality)
+{
+	if (Quality == 0)
+	{
+		//default 
+		Quality = 85;
+	}
+	else if (Quality == (int32)EImageCompressionQuality::Uncompressed)
+	{
+		// fix = 1 (Uncompressed) was previously treated as max-compress
+		Quality = 100;
+	}
+	else
+	{
+		ensure(Quality >= 1 && Quality <= 100);
+
+		#define JPEG_QUALITY_MIN 40
+		// JPEG should not be used below quality JPEG_QUALITY_MIN
+		Quality = FMath::Clamp(Quality, JPEG_QUALITY_MIN, 100);
+	}
+
+	if (CompressedData.Num() == 0)
+	{
+		check(RawData.Num());
+		check(Width > 0);
+		check(Height > 0);
+
+		// re-order components if required - JPEGs expect RGBA
+		if(Format == ERGBFormat::BGRA)
+		{
+			FColor* Colors = (FColor*)RawData.GetData();
+			const int32 NumColors = RawData.Num() / 4;
+			for(int32 ColorIndex = 0; ColorIndex < NumColors; ColorIndex++)
+			{
+				int32 tmp = Colors[ColorIndex].B;
+				Colors[ColorIndex].B = Colors[ColorIndex].R;
+				Colors[ColorIndex].R = tmp;
+			}
+		}
+
+
+		int data_size;
+		unsigned char* data;
+		oojpeg_encode_image_context_t ctx = oojpeg_encode_image_mem_threaded_start(Width, Height, NumComponents, RawData.GetData(), Quality, 0);
+		check(!ctx.error);
+		ParallelFor(ctx.jobs_to_run, [&](int idx) {
+			oojpeg_encode_image_thread_run(&ctx, idx);
+		});
+		int error = oojpeg_encode_image_mem_threaded_finish(&ctx, &data, &data_size);
+		check(!error);
+
+		CompressedData.Reset(data_size);
+		CompressedData.AddUninitialized(data_size);
+		FMemory::Memcpy(CompressedData.GetData(), data, data_size);
+		FMemory::Free(data);
+	}
+}
+
+
+void FOOJpegImageWrapper::Uncompress(const ERGBFormat InFormat, int32 InBitDepth)
+{
+	// Ensure we haven't already uncompressed the file.
+	if (RawData.Num() != 0)
+	{
+		return;
+	}
+
+	check(CompressedData.Num());
+
+	int32 NumColors;
+	int32 jpegWidth,jpegHeight;
+	uint8* OutData;
+	oojpeg_decode_context_t ctx = oojpeg_decode_mem_threaded_start(CompressedData.GetData(), CompressedData.Num(), 0);
+	if (ctx.error)
+	{
+		return;
+	}
+	ParallelFor(ctx.jobs_to_run, [&](int idx) {
+		oojpeg_decode_thread_run(&ctx, idx);
+	});
+	OutData = oojpeg_decode_threaded_finish(&ctx, &jpegWidth, &jpegHeight, &NumColors);
+
+	if (OutData)
+	{
+		Width = jpegWidth;
+		Height = jpegHeight;
+		RawData.Reset(Width * Height * NumColors);
+		RawData.AddUninitialized(Width * Height * NumColors);
+		FMemory::Memcpy(RawData.GetData(), OutData, RawData.Num());
+		FMemory::Free(OutData);
+	}
+	else
+	{
+		UE_LOG(LogImageWrapper, Error, TEXT("OOJPEG Decompress Error"));
+		RawData.Empty();
+	}
+}
+
+TArray64<uint8> FOOJpegImageWrapper::GetExportData(int32 Quality)
+{
+	uint8* Data;
+	int Size;
+	oojpeg_decode_mem_to_jpeg(CompressedData.GetData(), CompressedData.Num(), &Data, &Size, 0);
+
+	TArray64<uint8> ExportData;
+	ExportData.Reset();
+	ExportData.AddUninitialized(Size);
+	FMemory::Memcpy(ExportData.GetData(), Data, Size);
+	FMemory::Free(Data);
+	return ExportData;
+}
+
+#endif //WITH_OOJPEG
diff --git a/Engine/Source/Runtime/ImageWrapper/Private/Formats/OOJpegImageWrapper.h b/Engine/Source/Runtime/ImageWrapper/Private/Formats/OOJpegImageWrapper.h
new file mode 100644
--- /dev/null
+++ b/Engine/Source/Runtime/ImageWrapper/Private/Formats/OOJpegImageWrapper.h
@@ -0,0 +1,39 @@
+// Copyright Epic Games, Inc. All Rights Reserved.
+
+#pragma once
+
+#include "CoreTypes.h"
+#include "ImageWrapperBase.h"
+
+#if WITH_OOJPEG
+
+class FOOJpegImageWrapper
+	: public FImageWrapperBase
+{
+public:
+
+	/** Default constructor. */
+	// NumComponents == 1 for GrayscaleJPEG
+	FOOJpegImageWrapper(int32 InNumComponents = 4);
+
+	virtual ~FOOJpegImageWrapper();
+
+public:
+
+	//~ FImageWrapperBase interface
+
+	virtual bool SetCompressed(const void* InCompressedData, int64 InCompressedSize) override;
+	virtual void Uncompress(const ERGBFormat InFormat, int32 InBitDepth) override;
+	virtual void Compress(int32 Quality) override;
+	virtual TArray64<uint8> GetExportData(int32 Quality) override;
+	
+	virtual bool CanSetRawFormat(const ERGBFormat InFormat, const int32 InBitDepth) const override;
+	virtual ERawImageFormat::Type GetSupportedRawFormat(const ERawImageFormat::Type InFormat) const override;
+
+private:
+
+	int32 NumComponents;
+};
+
+
+#endif //WITH_OOJPEG
diff --git a/Engine/Source/Runtime/ImageWrapper/Private/ImageWrapperHelper.cpp b/Engine/Source/Runtime/ImageWrapper/Private/ImageWrapperHelper.cpp
--- a/Engine/Source/Runtime/ImageWrapper/Private/ImageWrapperHelper.cpp
+++ b/Engine/Source/Runtime/ImageWrapper/Private/ImageWrapperHelper.cpp
@@ -19,6 +19,12 @@ FStringView ImageWrapperHelper::GetFormatExtension(EImageFormat InImageFormat, b
 	case EImageFormat::GrayscaleJPEG:
 		StringExtension = TEXT(".jpg");
 		break;
+	case EImageFormat::OOJPEG:
+		StringExtension = TEXT(".ooj");
+		break;
+	case EImageFormat::GrayscaleOOJPEG:
+		StringExtension = TEXT(".ooj");
+		break;
 	case EImageFormat::BMP:
 		StringExtension = TEXT(".bmp");
 		break;
@@ -70,7 +76,11 @@ EImageFormat ImageWrapperHelper::GetImageFormat(FStringView StringExtention)
 	}
 	if (StringExtention.Equals(TEXT("jpg"), ESearchCase::IgnoreCase))
 	{
-		return EImageFormat::GrayscaleJPEG;
+		return EImageFormat::GrayscaleJPEG; // Note: This looks like a bug?
+	}
+	if (StringExtention.Equals(TEXT("ooj"), ESearchCase::IgnoreCase))
+	{
+		return EImageFormat::OOJPEG;
 	}
 	if (StringExtention.Equals(TEXT("bmp"), ESearchCase::IgnoreCase))
 	{
@@ -102,6 +112,6 @@ EImageFormat ImageWrapperHelper::GetImageFormat(FStringView StringExtention)
 const FStringView ImageWrapperHelper::GetImageFilesFilterString(bool bIncludeAllFiles)
 {
 	return (bIncludeAllFiles)
-		? TEXT("Image files (*.jpg; *.png; *.bmp; *.ico; *.exr; *.icns; *.jpeg; *.tga; *.hdr)|*.jpg; *.png; *.bmp; *.ico; *.exr; *.icns; *.jpeg; *.tga; *.hdr|All files (*.*)|*.*")
-		: TEXT("Image files (*.jpg; *.png; *.bmp; *.ico; *.exr; *.icns; *.jpeg; *.tga; *.hdr)|*.jpg; *.png; *.bmp; *.ico; *.exr; *.icns; *.jpeg; *.tga; *.hdr");
+		? TEXT("Image files (*.jpg; *.png; *.bmp; *.ico; *.exr; *.icns; *.jpeg; *.tga; *.hdr; *.ooj)|*.jpg; *.png; *.bmp; *.ico; *.exr; *.icns; *.jpeg; *.tga; *.hdr; *.ooj|All files (*.*)|*.*")
+		: TEXT("Image files (*.jpg; *.png; *.bmp; *.ico; *.exr; *.icns; *.jpeg; *.tga; *.hdr; *.ooj)|*.jpg; *.png; *.bmp; *.ico; *.exr; *.icns; *.jpeg; *.tga; *.hdr; *.ooj");
 }
\ No newline at end of file
diff --git a/Engine/Source/Runtime/ImageWrapper/Private/ImageWrapperModule.cpp b/Engine/Source/Runtime/ImageWrapper/Private/ImageWrapperModule.cpp
--- a/Engine/Source/Runtime/ImageWrapper/Private/ImageWrapperModule.cpp
+++ b/Engine/Source/Runtime/ImageWrapper/Private/ImageWrapperModule.cpp
@@ -11,6 +11,7 @@
 #include "Formats/IcnsImageWrapper.h"
 #include "Formats/IcoImageWrapper.h"
 #include "Formats/JpegImageWrapper.h"
+#include "Formats/OOJpegImageWrapper.h"
 #include "Formats/PngImageWrapper.h"
 #include "Formats/TgaImageWrapper.h"
 #include "Formats/TiffImageWrapper.h"
@@ -32,6 +33,7 @@ namespace
 	static const uint8 IMAGE_MAGIC_ICO[]  = {0x00, 0x00, 0x01, 0x00};
 	static const uint8 IMAGE_MAGIC_EXR[]  = {0x76, 0x2F, 0x31, 0x01};
 	static const uint8 IMAGE_MAGIC_ICNS[] = {0x69, 0x63, 0x6E, 0x73};
+	static const uint8 IMAGE_MAGIC_OOJPEG[] = {'O', 'O', 'J', 'P', 'E', 'G'};
 
 	// Binary for #?RADIANCE
 	static const uint8 IMAGE_MAGIC_HDR[] = {0x23, 0x3f, 0x52, 0x41, 0x44, 0x49, 0x41, 0x4e, 0x43, 0x45, 0x0a};
@@ -99,6 +101,16 @@ public:
 			break;
 #endif	//WITH_UNREALJPEG
 
+#if WITH_OOJPEG
+		case EImageFormat::OOJPEG:
+			ImageWrapper = MakeShared<FOOJpegImageWrapper>();
+			break;
+
+		case EImageFormat::GrayscaleOOJPEG:
+			ImageWrapper = MakeShared<FOOJpegImageWrapper>(1);
+			break;
+#endif // WITH_OOJPEG
+
 		case EImageFormat::BMP:
 			ImageWrapper = MakeShared<FBmpImageWrapper>();
 			break;
@@ -156,6 +168,10 @@ public:
 		{
 			Format = EImageFormat::JPEG; // @Todo: Should we detect grayscale vs non-grayscale?
 		}
+		else if (StartsWith((uint8*)CompressedData, CompressedSize, IMAGE_MAGIC_OOJPEG))
+		{
+			Format = EImageFormat::OOJPEG; // @Todo: Should we detect grayscale vs non-grayscale?
+		}
 		else if (StartsWith((uint8*) CompressedData, CompressedSize, IMAGE_MAGIC_BMP))
 		{
 			Format = EImageFormat::BMP;
@@ -207,6 +223,8 @@ public:
 		case EImageFormat::PNG: return TEXT("png");
 		case EImageFormat::JPEG: return TEXT("jpg");
 		case EImageFormat::GrayscaleJPEG: return TEXT("jpg");
+		case EImageFormat::OOJPEG: return TEXT("ooj");
+		case EImageFormat::GrayscaleOOJPEG: return TEXT("ooj");
 		case EImageFormat::BMP: return TEXT("bmp");
 		case EImageFormat::ICO: return TEXT("ico");
 		case EImageFormat::EXR: return TEXT("exr");
@@ -238,6 +256,10 @@ public:
 		{
 			return EImageFormat::JPEG;
 		}
+		else if (FCString::Stricmp(Name, TEXT("ooj")) == 0)
+		{
+			return EImageFormat::OOJPEG;
+		}
 		else if ( FCString::Stricmp(Name,TEXT("bmp")) == 0 )
 		{
 			return EImageFormat::BMP;
diff --git a/Engine/Source/Runtime/ImageWrapper/Public/IImageWrapper.h b/Engine/Source/Runtime/ImageWrapper/Public/IImageWrapper.h
--- a/Engine/Source/Runtime/ImageWrapper/Public/IImageWrapper.h
+++ b/Engine/Source/Runtime/ImageWrapper/Public/IImageWrapper.h
@@ -59,6 +59,12 @@ enum class EImageFormat : int8
 
 	/** DirectDraw Surface */
 	DDS,
+
+	/** Oodle JPEG format. */
+	OOJPEG,
+
+	/** Single channel Oodle JPEG. */
+	GrayscaleOOJPEG,
 };
 
 
@@ -181,6 +187,18 @@ public:
 	 */
 	virtual TArray64<uint8> GetCompressed(int32 Quality = 0) = 0;
 
+	/**
+	 * Gets the data for export. 
+	 * Usually the same thing as GetCompressed.
+	 * 
+	 * @return Array of the data to export.  returns empty array on failure
+	 * 
+	 */
+	virtual TArray64<uint8> GetExportData(int32 Quality = 0)
+	{
+		return GetCompressed(Quality);
+	}
+
 	/**
 	* GetRaw after SetCompressed
 	* fills the raw data in the native format/depth contained in the file
diff --git a/Engine/Source/Runtime/OodleJPEG/OodleJPEG.Build.cs b/Engine/Source/Runtime/OodleJPEG/OodleJPEG.Build.cs
new file mode 100644
--- /dev/null
+++ b/Engine/Source/Runtime/OodleJPEG/OodleJPEG.Build.cs
@@ -0,0 +1,38 @@
+// Copyright Epic Games, Inc. All Rights Reserved.
+using UnrealBuildTool;
+using System.IO;
+
+public class OodleJPEG : ModuleRules
+{
+	public OodleJPEG(ReadOnlyTargetRules Target) : base(Target)
+	{
+		Type = ModuleType.External;
+
+		bool bWithLib = false;
+		if (Target.Platform == UnrealTargetPlatform.Win64 && Target.WindowsPlatform.Architecture != UnrealArch.Arm64)
+		{
+			string LibPath = Path.Combine(ModuleDirectory, "Lib");
+			PublicAdditionalLibraries.Add(Path.Combine(LibPath, "oojpeg_w64.lib"));
+			bWithLib = true;
+		}
+		else if (Target.Platform == UnrealTargetPlatform.Mac)
+		{
+			string LibPath = Path.Combine(ModuleDirectory, "Lib");
+			PublicAdditionalLibraries.Add(Path.Combine(LibPath, "oojpeg_mac.a"));
+			bWithLib = true;
+		}
+		else if (Target.IsInPlatformGroup(UnrealPlatformGroup.Unix))
+		{
+			string LibPath = Path.Combine(ModuleDirectory, "Lib/Unix", Target.Architecture.LinuxName);
+			PublicAdditionalLibraries.Add(Path.Combine(LibPath, "oojpeg_linux.a"));
+			bWithLib = true;
+		}
+
+		if (bWithLib)
+		{
+			PublicSystemIncludePaths.Add(Path.Combine(ModuleDirectory, "Source"));
+		}
+
+		PublicDefinitions.Add("WITH_OOJPEG=" + (bWithLib ? '1' : '0'));
+	}
+}
diff --git a/Engine/Source/Runtime/OodleJPEG/Source/oojpeg.h b/Engine/Source/Runtime/OodleJPEG/Source/oojpeg.h
new file mode 100644
--- /dev/null
+++ b/Engine/Source/Runtime/OodleJPEG/Source/oojpeg.h
@@ -0,0 +1,284 @@
+// Copyright Epic Games, Inc. All Rights Reserved.
+/*
+
+OOJPEG - A JPEG Super-compressor using an oodle backend as the entropy coder.
+
+Can encode from a JPEG file to an OOJPEG file, and can decode directly to a JPEG file or a raw image.
+Can also encode raw images directly to OOJPEG files.
+
+API has simple APIs and threaded APIs. 
+
+The simple APIs is as follows:
+-----------------------------
+
+The following functions encode from a JPEG file to a OOJPEG file...
+
+  int oojpeg_encode_jpeg_mem(const unsigned char *jpeg_data, int jpeg_data_size, unsigned char **out_oojpeg_data, int *out_oojpeg_size, int flags);
+
+The following functions encode from a raw RGBA image to a OOJPEG file...
+
+  int oojpeg_encode_to_mem(unsigned char **out_data, int *out_size, int width, int height, int channels, const unsigned char *rgba, int quality, int flags);
+
+The following functions decode from a OOJPEG file back to a JPEG file...
+
+  int oojpeg_decode_mem_to_jpeg(const unsigned char *data, int size, unsigned char **out_data, int *out_size, int flags);
+
+The following functions decode from a OOJPEG file to a raw RGBA image...
+
+  unsigned char *oojpeg_decode_mem(const unsigned char *data, int size, int *width, int *height, int *channels, int flags);
+
+The threaded APIs is as follows:
+-------------------------------
+
+The following functions encode from a JPEG file to a OOJPEG file...
+
+  oojpeg_encode_context_t oojpeg_encode_jpeg_mem_threaded_start(const unsigned char *data, int size, int flags);
+  int oojpeg_encode_jpeg_thread_run(oojpeg_encode_context_t *context, int job_idx);
+  int oojpeg_encode_jpeg_mem_threaded_finish(oojpeg_encode_context_t *context, unsigned char **out_data, int *out_size);
+
+The following functions encode from a raw RGBA image to a OOJPEG file...
+
+  oojpeg_encode_image_context_t oojpeg_encode_image_mem_threaded_start(int width, int height, int comp, const unsigned char *in_data, int quality, int flags);
+  oojpeg_encode_image_context_t oojpeg_encode_image_threaded_start(const oojpeg_io_callbacks_t *io, void *io_user, int width, int height, int comp, const unsigned char *in_data, int quality, int flags);
+  int oojpeg_encode_image_thread_run(oojpeg_encode_image_context_t *ctx, int job_idx);
+  int oojpeg_encode_image_threaded_finish(oojpeg_encode_image_context_t *ctx);
+  int oojpeg_encode_image_mem_threaded_finish(oojpeg_encode_image_context_t *ctx, unsigned char **out_data, int *out_size);
+
+The following functions decode from a OOJPEG file to a raw RGBA image...
+
+  oojpeg_decode_context_t oojpeg_decode_mem_threaded_start(const unsigned char *data, int size, int flags);
+  oojpeg_decode_context_t oojpeg_decode_threaded_start(const oojpeg_io_callbacks_t *io, void *io_user, int flags);
+  int oojpeg_decode_thread_run(oojpeg_decode_context_t *context, int job_idx);
+  unsigned char *oojpeg_decode_threaded_finish(oojpeg_decode_context_t *context, int *out_width, int *out_height, int *out_comp);
+
+You use the threaded APIs by first calling the _start function, which gives you a context. You then pass that context into the thread_run function N times,
+where N is the value of the context.jobs_to_run variable. You then wait for all the threads to finish, and then call the _finish function to get the output data.
+
+Example:
+
+      oojpeg_encode_context_t ctx = oojpeg_encode_jpeg_mem_threaded_start(input, output, flags);
+      #pragma omp parallel for
+      for(int i = 0; i < ctx.jobs_to_run; ++i) {
+        oojpeg_encode_jpeg_thread_run(&ctx, i);
+      }
+      oojpeg_encode_jpeg_file_threaded_finish(&ctx);
+
+Other Notes:
+
+* Custom Allocators can be set via oojpeg_set_alloc() function.
+* Custom compression/decompression (must be set the same in both the encoder and decoder) can be set via oojpeg_set_compression() function.
+
+*/
+#ifndef OOJPEG_H
+#define OOJPEG_H
+
+#ifdef __cplusplus
+extern "C" {
+#endif
+
+// Error codes
+typedef enum {
+  OOJPEG_ERR_NONE                = 0, // No Error, everything OK! :)
+  OOJPEG_ERR_NOT_A_JPEG          = 1, // Not a JPEG file. Happens if you try to encode from a non-JPEG file.
+  OOJPEG_ERR_CANNOT_OPEN_FILE    = 2, // Cannot open file. Happens if the file routines cannot open the files you specified!
+  OOJPEG_ERR_NOT_OOJPEG_FILE     = 3, // Not a Oojpeg file. Happens if you try to decode a non-Oojpeg file.
+  OOJPEG_ERR_CORRUPT_OOJPEG_FILE = 4, // Corrupt Oojpeg file. Happens if the Oojpeg file is corrupt.
+  OOJPEG_ERR_WRITE_FAILED        = 5, // Write failed. Happens if the file routines cannot write to the files you specified!
+  OOJPEG_ERR_OUT_OF_MEMORY       = 6, // Out of memory. Happens if the Oojpeg library cannot allocate enough memory.
+  OOJPEG_ERR_UNSUPPORTED_COLORSPACE = 7, // Unsupported colorspace. Happens when exporting if the JPEG file is not YCbCr.
+  OOJPEG_ERR_READ_FAILED         = 8, // Read failed. Happens if the file routines cannot read from the files you specified!
+  OOJPEG_ERR_UNSUPPORTED         = 9, // Other Unsupported. Happens usually when you try to decode to jpeg in an uncommon format.
+} oojpeg_error_t;
+
+// Flags to be used in compression or decompression
+typedef enum {
+  OOJPEG_FLAG_NONE            = 0x0, // no flags
+  OOJPEG_FLAG_NO_COMPRESSION  = 0x1, // don't compress the coefficients, just store them
+} oojpeg_flags_t;
+
+enum {
+  OOJPEG_MAX_SPLITS = 16,
+};
+
+// IO callbacks
+//  For use with oojpeg_*() to specify any input or output method. All methods internally use this!
+typedef struct {
+  size_t (*write)(void *user, const void *data, size_t size);  // write 'data' with 'size' bytes.  return number of bytes actually written
+  size_t (*read)(void *user, void *data, size_t size);  // read 'data' with 'size' bytes.  return number of bytes actually read
+  int (*eof)(void *user);  // return 1 if end-of-file, 0 otherwise
+} oojpeg_io_callbacks_t;
+
+// For using memory buffers as input or output.
+typedef struct {
+  char *base;
+  char *ptr;
+  size_t size;
+} oojpeg_io_buffer_t;
+
+extern const oojpeg_io_callbacks_t s_oojpeg_buffer_fns;
+
+// For setting custom memory allocators.
+typedef struct {
+  void *(*alloc)(size_t size);  // allocate 'size' bytes.  return pointer to allocated memory
+  void (*free)(void *ptr);  // free memory at 'ptr'
+  void *(*realloc)(void *ptr, size_t size);  // reallocate 'ptr' to 'size' bytes.  return pointer to reallocated memory
+} oojpeg_alloc_t;
+
+typedef struct {
+    void *data;
+    size_t size;
+} oojpeg_buffer_t;
+
+// compression callbacks, for using custom compression methods, such as zlib
+typedef struct {
+  oojpeg_buffer_t (*compress)(const void *data, size_t size);
+  int (*decompress)(const void *data, size_t size, void *out_data, size_t out_size);
+} oojpeg_compression_t;
+
+// Set custom allocators.  If not set, malloc/free/realloc will be used.
+void oojpeg_set_alloc(const oojpeg_alloc_t *alloc);
+// Set custom compression methods.  If not set, oodle compression methods will be used.
+void oojpeg_set_compression(const oojpeg_compression_t *compression);
+
+// An internal structure used to store the header of a Oojpeg file.
+typedef struct oojpeg_ihdr_t {
+  unsigned version;
+  unsigned width;
+  unsigned height;
+  unsigned char bit_depth;
+  unsigned char comp;
+  unsigned char method;
+  unsigned char num_splits;
+} oojpeg_ihdr_t;
+
+// An internal structure used to store the header of a Oojpeg file.
+typedef struct oojpeg_lossy_hdr_t {
+  unsigned short coef_dims[4][2];
+  // Quantization tables
+  unsigned short dequant[4][64];
+  char app14_color_transform; // valid values are 0(CMYK),1(YCCK),2(YCbCrA)
+  char jfif;
+  char comp_id[4];
+} oojpeg_lossy_hdr_t;
+
+typedef struct {
+  int y_start[4];
+  int y_end[4];
+  int width[4];
+  int height[4];
+  int num_blocks;
+  short *xcoefs;
+  char *split_xcoefs;
+  unsigned char *oodle_buf;
+  int oodle_size;
+  int coefs_size;
+} oojpeg_internal_split_t;
+
+typedef struct {
+  int error;
+  int split_height_min;
+  int split_height[4];
+  int split_blocks[4];
+  int total_split_blocks;
+  oojpeg_internal_split_t splits[OOJPEG_MAX_SPLITS];
+} oojpeg_internal_splits_t;
+
+typedef struct {
+  int error; // If there was an error, it is stored here.
+  int jobs_to_run; // How many times to run oojpeg_decode_thread_run!
+
+  // arguments
+  const oojpeg_io_callbacks_t *io;
+  void *io_user;
+  int flags;
+
+  // internal state
+  int width, height, comp;
+  int num_splits;
+
+  oojpeg_lossy_hdr_t lhdr;
+  oojpeg_internal_splits_t s;
+} oojpeg_encode_context_t;
+
+typedef struct {
+  int error; // If there was an error, it is stored here.
+  int jobs_to_run; // How many times to run oojpeg_decode_thread_run!
+
+  // arguments
+  int flags;
+
+  // internal state
+  oojpeg_ihdr_t ihdr;
+  oojpeg_lossy_hdr_t lhdr;
+  oojpeg_internal_splits_t spx;
+  int yuva_stride[4];
+  int yuva_height[4];
+  unsigned char *yuva[4];
+} oojpeg_decode_context_t;
+
+typedef struct {
+  int error; // If there was an error, it is stored here.
+  int jobs_to_run; // How many times to run oojpeg_decode_thread_run!
+
+  // arguments
+  const oojpeg_io_callbacks_t *io;
+  void *io_user;
+  int width, height, comp;
+  const unsigned char *in_data;
+  int flags;
+
+  // internal state
+  int num_splits;
+  oojpeg_internal_splits_t spx;
+  oojpeg_lossy_hdr_t lhdr;
+	float fdtbl_Y[64];
+  float fdtbl_UV[64];
+} oojpeg_encode_image_context_t;
+
+// Encode a JPEG file to Oojpeg format.
+// This second function operates totally in memory.
+int oojpeg_encode_jpeg_mem(const unsigned char *jpeg_data, int jpeg_data_size, unsigned char **out_oojpeg_data, int *out_oojpeg_size, int flags);
+
+// Threaded API
+// In this API you first call start, and it does some synchronous work and gives you a state context to pass to the run function.
+// Then inside threads, you call the run function with job_idx == 0 .. context.jobs_to_run
+// Then you wait for all the threads to finish processing all the jobs.
+// Then you call finish to get the output data.
+oojpeg_encode_context_t oojpeg_encode_jpeg_mem_threaded_start(const unsigned char *data, int size, int flags);
+int oojpeg_encode_jpeg_thread_run(oojpeg_encode_context_t *context, int job_idx);
+int oojpeg_encode_jpeg_mem_threaded_finish(oojpeg_encode_context_t *context, unsigned char **out_data, int *out_size);
+
+// Encode a raw image to a Oojpeg file.
+int oojpeg_encode_to_mem(unsigned char **out_data, int *out_size, int width, int height, int channels, const unsigned char *rgba, int quality, int flags);
+
+oojpeg_encode_image_context_t oojpeg_encode_image_mem_threaded_start(int width, int height, int comp, const unsigned char *in_data, int quality, int flags);
+oojpeg_encode_image_context_t oojpeg_encode_image_threaded_start(const oojpeg_io_callbacks_t *io, void *io_user, int width, int height, int comp, const unsigned char *in_data, int quality, int flags);
+int oojpeg_encode_image_thread_run(oojpeg_encode_image_context_t *ctx, int job_idx);
+int oojpeg_encode_image_threaded_finish(oojpeg_encode_image_context_t *ctx);
+int oojpeg_encode_image_mem_threaded_finish(oojpeg_encode_image_context_t *ctx, unsigned char **out_data, int *out_size);
+
+// Decode a OOJPEG file to a raw image.
+// Interface is like stbi_load_from_file, stbi_load_from_memory, stbi_load, etc.
+unsigned char *oojpeg_decode_mem(const unsigned char *data, int size, int *width, int *height, int *channels, int flags);
+unsigned char *oojpeg_decode(const oojpeg_io_callbacks_t *io, void *io_user, int *out_width, int *out_height, int *out_comp, int flags);
+
+// Threaded API
+// In this API you first call start, and it does some synchronous work and gives you a state context to pass to the run function.
+// Then inside threads, you call the run function with job_idx == 0 .. context.jobs_to_run
+// Then you wait for all the threads to finish processing all the jobs.
+// Then you call finish to get the output data.
+
+oojpeg_decode_context_t oojpeg_decode_mem_threaded_start(const unsigned char *data, int size, int flags);
+oojpeg_decode_context_t oojpeg_decode_threaded_start(const oojpeg_io_callbacks_t *io, void *io_user, int flags);
+int oojpeg_decode_thread_run(oojpeg_decode_context_t *context, int job_idx);
+unsigned char *oojpeg_decode_threaded_finish(oojpeg_decode_context_t *context, int *out_width, int *out_height, int *out_comp);
+
+// Decode a OOJPEG file to a JPEG file.
+int oojpeg_decode_to_jpeg(const oojpeg_io_callbacks_t *in_io, void *in_io_user, const oojpeg_io_callbacks_t *out_io, void *out_io_user, int flags);
+int oojpeg_decode_mem_to_jpeg(const unsigned char *data, int size, unsigned char **out_data, int *out_size, int flags);
+
+#ifdef __cplusplus
+}
+#endif
+
+#endif // OOJPEG_H
