Oodle Jpeg: Fixed old-style header warning double include protection
#p4v-preflight-copy 27847780

diff --git a/Engine/Source/Runtime/OodleJPEG/Source/oojpeg.h b/Engine/Source/Runtime/OodleJPEG/Source/oojpeg.h
--- a/Engine/Source/Runtime/OodleJPEG/Source/oojpeg.h
+++ b/Engine/Source/Runtime/OodleJPEG/Source/oojpeg.h
@@ -69,8 +69,7 @@ Other Notes:
 * Custom compression/decompression (must be set the same in both the encoder and decoder) can be set via oojpeg_set_compression() function.
 
 */
-#ifndef OOJPEG_H
-#define OOJPEG_H
+#pragma once
 
 #ifdef __cplusplus
 extern "C" {
@@ -281,4 +280,3 @@ int oojpeg_decode_mem_to_jpeg(const unsigned char *data, int size, unsigned char
 }
 #endif
 
-#endif // OOJPEG_H
